Update default copyright notices for C# projects.

diff --git a/Engine/Source/Programs/HTML5/HTML5LaunchHelper/Properties/AssemblyInfo.cs b/Engine/Source/Programs/HTML5/HTML5LaunchHelper/Properties/AssemblyInfo.cs
--- a/Engine/Source/Programs/HTML5/HTML5LaunchHelper/Properties/AssemblyInfo.cs
+++ b/Engine/Source/Programs/HTML5/HTML5LaunchHelper/Properties/AssemblyInfo.cs
@@ -12,7 +12,7 @@ using System.Runtime.InteropServices;
 [assembly: AssemblyConfiguration("")]
 [assembly: AssemblyCompany("")]
 [assembly: AssemblyProduct("HTML5LaunchHelper")]
-[assembly: AssemblyCopyright("Copyright ©  2015")]
+[assembly: AssemblyCopyright("Copyright 1998-2019 Epic Games, Inc. All rights reserved.")]
 [assembly: AssemblyTrademark("")]
 [assembly: AssemblyCulture("")]
 
diff --git a/Engine/Source/Programs/UnrealGameSync/InstallerCustomAction/Properties/AssemblyInfo.cs b/Engine/Source/Programs/UnrealGameSync/InstallerCustomAction/Properties/AssemblyInfo.cs
--- a/Engine/Source/Programs/UnrealGameSync/InstallerCustomAction/Properties/AssemblyInfo.cs
+++ b/Engine/Source/Programs/UnrealGameSync/InstallerCustomAction/Properties/AssemblyInfo.cs
@@ -11,7 +11,7 @@ using System.Runtime.InteropServices;
 [assembly: AssemblyDescription("")]
 [assembly: AssemblyCompany("")]
 [assembly: AssemblyProduct("InstallerCustomAction")]
-[assembly: AssemblyCopyright("Copyright ©  2015")]
+[assembly: AssemblyCopyright("Copyright 1998-2019 Epic Games, Inc. All rights reserved.")]
 [assembly: AssemblyTrademark("")]
 [assembly: AssemblyCulture("")]
 
diff --git a/Engine/Source/Programs/UnrealProjectTemplate/UE4TemplateProject/Properties/AssemblyInfo.cs b/Engine/Source/Programs/UnrealProjectTemplate/UE4TemplateProject/Properties/AssemblyInfo.cs
--- a/Engine/Source/Programs/UnrealProjectTemplate/UE4TemplateProject/Properties/AssemblyInfo.cs
+++ b/Engine/Source/Programs/UnrealProjectTemplate/UE4TemplateProject/Properties/AssemblyInfo.cs
@@ -12,7 +12,7 @@ using System.Runtime.InteropServices;
 [assembly: AssemblyConfiguration("")]
 [assembly: AssemblyCompany("")]
 [assembly: AssemblyProduct("UE4TemplateProject")]
-[assembly: AssemblyCopyright("Copyright ©  2015")]
+[assembly: AssemblyCopyright("Copyright 1998-2019 Epic Games, Inc. All rights reserved.")]
 [assembly: AssemblyTrademark("")]
 [assembly: AssemblyCulture("")]
 
diff --git a/Engine/Source/Programs/UnrealProjectTemplate/WizardAssembly/Properties/AssemblyInfo.cs b/Engine/Source/Programs/UnrealProjectTemplate/WizardAssembly/Properties/AssemblyInfo.cs
--- a/Engine/Source/Programs/UnrealProjectTemplate/WizardAssembly/Properties/AssemblyInfo.cs
+++ b/Engine/Source/Programs/UnrealProjectTemplate/WizardAssembly/Properties/AssemblyInfo.cs
@@ -14,7 +14,7 @@ using System.Windows;
 [assembly: AssemblyConfiguration("")]
 [assembly: AssemblyCompany("")]
 [assembly: AssemblyProduct("WizardAssembly")]
-[assembly: AssemblyCopyright("Copyright ©  2015")]
+[assembly: AssemblyCopyright("Copyright 1998-2019 Epic Games, Inc. All rights reserved.")]
 [assembly: AssemblyTrademark("")]
 [assembly: AssemblyCulture("")]
 
diff --git a/Engine/Source/Programs/UnrealSwarm/SwarmCommonUtils/Properties/AssemblyInfo.cs b/Engine/Source/Programs/UnrealSwarm/SwarmCommonUtils/Properties/AssemblyInfo.cs
--- a/Engine/Source/Programs/UnrealSwarm/SwarmCommonUtils/Properties/AssemblyInfo.cs
+++ b/Engine/Source/Programs/UnrealSwarm/SwarmCommonUtils/Properties/AssemblyInfo.cs
@@ -12,7 +12,7 @@ using System.Runtime.InteropServices;
 [assembly: AssemblyConfiguration("")]
 [assembly: AssemblyCompany("")]
 [assembly: AssemblyProduct("SwarmCommonUtils")]
-[assembly: AssemblyCopyright("Copyright ©  2014")]
+[assembly: AssemblyCopyright("Copyright 1998-2019 Epic Games, Inc. All rights reserved.")]
 [assembly: AssemblyTrademark("")]
 [assembly: AssemblyCulture("")]
 
